Copy from //Tasks/UE4/Partner-vrCluster
nDisplayLauncher - Removed OpenGL3 RHI option from the launcher

diff --git a/Engine/Source/Programs/nDisplayLauncher/Cluster/Launcher/Launcher.cs b/Engine/Source/Programs/nDisplayLauncher/Cluster/Launcher/Launcher.cs
--- a/Engine/Source/Programs/nDisplayLauncher/Cluster/Launcher/Launcher.cs
+++ b/Engine/Source/Programs/nDisplayLauncher/Cluster/Launcher/Launcher.cs
@@ -58,7 +58,6 @@ namespace nDisplayLauncher.Cluster
 		#region Launcher_Params
 		private Dictionary<string, string> _RenderApiParams = new Dictionary<string, string>
 		{
-			{"OpenGL3",    "-opengl3" },
 			{"OpenGL4",    "-opengl4" },
 			{"DirectX 11", "-dx11" },
 			{"DirectX 12", "-dx12" }
@@ -451,7 +450,7 @@ namespace nDisplayLauncher.Cluster
 			}
 			catch (Exception)
 			{
-				SelectedRenderApiParam = RenderApiParams.SingleOrDefault(x => x.Key == "OpenGL3");
+				SelectedRenderApiParam = RenderApiParams.SingleOrDefault(x => x.Key == "DirectX 11");
 			}
 
 			try
