nDisplayLauncher: Fixed issue with spaces in executable path

diff --git a/Engine/Source/Programs/nDisplayLauncher/Cluster/Launcher/Launcher.StartApp.cs b/Engine/Source/Programs/nDisplayLauncher/Cluster/Launcher/Launcher.StartApp.cs
--- a/Engine/Source/Programs/nDisplayLauncher/Cluster/Launcher/Launcher.StartApp.cs
+++ b/Engine/Source/Programs/nDisplayLauncher/Cluster/Launcher/Launcher.StartApp.cs
@@ -6,6 +6,7 @@ using System.IO;
 
 using nDisplayLauncher.Cluster.Config;
 using nDisplayLauncher.Cluster.Config.Entity;
+using nDisplayLauncher.Helpers;
 using nDisplayLauncher.Log;
 
 
@@ -21,28 +22,6 @@ namespace nDisplayLauncher.Cluster
 				return;
 			}
 
-			string Application = SelectedApplication;
-			int FirstSpacePos = Application.IndexOf(' ');
-			if (FirstSpacePos > 0)
-			{
-				Application = SelectedApplication.Substring(0, FirstSpacePos);
-			}
-			if (!File.Exists(Application))
-			{
-				AppLogger.Log("Application not found: " + Application);
-				return;
-			}
-
-			// Update config files before application start
-			HashSet<string> NodesSent = new HashSet<string>();
-			foreach (EntityClusterNode Node in Config.ClusterNodes.Values)
-			{
-				if (!NodesSent.Contains(Node.Addr))
-				{
-					NodesSent.Add(Node.Addr);
-				}
-			}
-
 			// Send start command to the listeners
 			foreach (EntityClusterNode Node in Config.ClusterNodes.Values)
 			{
@@ -57,27 +36,14 @@ namespace nDisplayLauncher.Cluster
 			string Application = SelectedApplication;
 			string ExtraAppParams = string.Empty;
 
-			// Executable
-			int FirstSpacePos = SelectedApplication.IndexOf(' ');
-			if (FirstSpacePos > 0)
-			{
-				Application = SelectedApplication.Substring(0, FirstSpacePos);
-				ExtraAppParams = SelectedApplication.Substring(FirstSpacePos +1);
-			}
-			commandCmd = string.Format("{0} \"{1}\"", CommandStartApp, Application);
+			// Start command at first
+			commandCmd = CommandStartApp;
 
-			if (!string.IsNullOrWhiteSpace(ExtraAppParams))
+			// Get file(s)
+			List<string> Params = SplitAppLineParameters(SelectedApplication);
+			foreach (string Param in Params)
 			{
-				string FirstArg = ExtraAppParams;
-				string RemArgs = String.Empty;
-				FirstSpacePos = ExtraAppParams.IndexOf(' ');
-				if (FirstSpacePos > 0)
-				{
-					FirstArg = ExtraAppParams.Substring(0, FirstSpacePos);
-					RemArgs = ExtraAppParams.Substring(FirstSpacePos +1);
-				}
-
-				commandCmd = string.Format("{0} \"{1}\" {2}", commandCmd, FirstArg, RemArgs);
+				commandCmd = string.Format("{0} {1}", commandCmd, Param);
 			}
 
 			// Custom common arguments
@@ -172,5 +138,46 @@ namespace nDisplayLauncher.Cluster
 
 			return commandCmd;
 		}
+
+		private List<string> SplitAppLineParameters(string Line)
+		{
+			const string ExecExtension = ".exe";
+			const string DashGameParam = "-game";
+
+			List<string> Params = new List<string>();
+
+			// Detect either we have -game line or single .exe/.bat/.cmd
+			bool IsDashGame = StringHelper.Contains(Line, DashGameParam, StringComparison.CurrentCultureIgnoreCase);
+
+			// In -game mode we have two files. One is the Editor executable, another is a project.
+			if (IsDashGame)
+			{
+				// Extract Editor and project files
+				int Idx = Line.IndexOf(ExecExtension, StringComparison.CurrentCultureIgnoreCase);
+				if (Idx > 0)
+				{
+					// 1. Add executable to output
+					string EditorFile = string.Format("\"{0}\"", Line.Substring(0, Idx + ExecExtension.Length).Trim());
+					Params.Add(EditorFile);
+
+					// 2. Remove everything except of .uproject file and store it to output
+					Line = Line.Remove(0, Idx + ExecExtension.Length);
+					Line = Line.Remove(Line.IndexOf(DashGameParam));
+					string ProjectFile = string.Format("\"{0}\"", Line.Trim());
+					Params.Add(ProjectFile);
+
+					// 3. Finally, add -game
+					Params.Add(DashGameParam);
+				}
+			}
+			// Otherwise we have one executable so nothing special to do
+			else
+			{
+				string ExecFile = string.Format("\"{0}\"", Line.Trim());
+				Params.Add(ExecFile);
+			}
+
+			return Params;
+		}
 	}
 }
diff --git a/Engine/Source/Programs/nDisplayLauncher/Helpers/StringHelper.cs b/Engine/Source/Programs/nDisplayLauncher/Helpers/StringHelper.cs
new file mode 100644
--- /dev/null
+++ b/Engine/Source/Programs/nDisplayLauncher/Helpers/StringHelper.cs
@@ -0,0 +1,15 @@
+using System;
+
+
+namespace nDisplayLauncher.Helpers
+{
+	// String helpers
+	public static class StringHelper
+	{
+		// Extended version of string.Contains with case checking
+		public static bool Contains(string Source, string ToCheck, StringComparison Comp)
+		{
+			return Source?.IndexOf(ToCheck, Comp) >= 0;
+		}
+	}
+}
diff --git a/Engine/Source/Programs/nDisplayLauncher/nDisplayLauncher.csproj b/Engine/Source/Programs/nDisplayLauncher/nDisplayLauncher.csproj
--- a/Engine/Source/Programs/nDisplayLauncher/nDisplayLauncher.csproj
+++ b/Engine/Source/Programs/nDisplayLauncher/nDisplayLauncher.csproj
@@ -151,6 +151,7 @@
     <Compile Include="Cluster\Config\Entity\EntityWindow.cs" />
     <Compile Include="Cluster\Config\Entity\EntityClusterNode.cs" />
     <Compile Include="Cluster\Config\Parser.cs" />
+    <Compile Include="Helpers\StringHelper.cs" />
     <Compile Include="UIControl\AutoScroller.cs" />
     <Compile Include="Settings\RegistrySaver.cs" />
     <Compile Include="UIControl\UnselectableListBox.cs" />
