Fix CIS warning

diff --git a/Engine/Source/Programs/nDisplayLauncher/Helpers/StringHelper.cs b/Engine/Source/Programs/nDisplayLauncher/Helpers/StringHelper.cs
--- a/Engine/Source/Programs/nDisplayLauncher/Helpers/StringHelper.cs
+++ b/Engine/Source/Programs/nDisplayLauncher/Helpers/StringHelper.cs
@@ -1,6 +1,6 @@
+// Copyright 1998-2019 Epic Games, Inc. All Rights Reserved.
 using System;
 
-
 namespace nDisplayLauncher.Helpers
 {
 	// String helpers
