Remove Engine version from nDisplayLauncher window title
Can't update that every version to match with engine.

diff --git a/Engine/Source/Programs/nDisplayLauncher/MainWindow.xaml.cs b/Engine/Source/Programs/nDisplayLauncher/MainWindow.xaml.cs
--- a/Engine/Source/Programs/nDisplayLauncher/MainWindow.xaml.cs
+++ b/Engine/Source/Programs/nDisplayLauncher/MainWindow.xaml.cs
@@ -26,7 +26,7 @@ namespace nDisplayLauncher
 
 		private void Window_Loaded(object sender, RoutedEventArgs e)
 		{
-			this.Title = "nDisplay Launcher - UE4.23";
+			this.Title = "nDisplay Launcher";
 		}
 
 		private void Exit(object sender, RoutedEventArgs e)
