Added RenderSync log category to the nDisplayLauncher 'Log' tab

diff --git a/Engine/Source/Programs/nDisplayLauncher/Cluster/Launcher/Launcher.StartApp.cs b/Engine/Source/Programs/nDisplayLauncher/Cluster/Launcher/Launcher.StartApp.cs
--- a/Engine/Source/Programs/nDisplayLauncher/Cluster/Launcher/Launcher.StartApp.cs
+++ b/Engine/Source/Programs/nDisplayLauncher/Cluster/Launcher/Launcher.StartApp.cs
@@ -118,7 +118,8 @@ namespace nDisplayLauncher.Cluster
 					"LogDisplayClusterNetwork {8}, " +
 					"LogDisplayClusterNetworkMsg {9}, " +
 					"LogDisplayClusterRender {10}, " +
-					"LogDisplayClusterBlueprint {11}"
+					"LogDisplayClusterRenderSync {11}, " +
+					"LogDisplayClusterBlueprint {12}"
 					, SelectedVerbocityPlugin
 					, SelectedVerbocityEngine
 					, SelectedVerbocityConfig
@@ -130,6 +131,7 @@ namespace nDisplayLauncher.Cluster
 					, SelectedVerbocityNetwork
 					, SelectedVerbocityNetworkMsg
 					, SelectedVerbocityRender
+					, SelectedVerbocityRenderSync
 					, SelectedVerbocityBlueprint);
 
 				commandCmd = string.Format("{0} -LogCmds=\"{1}\"", commandCmd, LogCmds);
diff --git a/Engine/Source/Programs/nDisplayLauncher/Cluster/Launcher/Launcher.cs b/Engine/Source/Programs/nDisplayLauncher/Cluster/Launcher/Launcher.cs
--- a/Engine/Source/Programs/nDisplayLauncher/Cluster/Launcher/Launcher.cs
+++ b/Engine/Source/Programs/nDisplayLauncher/Cluster/Launcher/Launcher.cs
@@ -391,6 +391,17 @@ namespace nDisplayLauncher.Cluster
 				RegistrySaver.UpdateRegistry(RegistrySaver.RegCategoryLogParams, RegistrySaver.RegLogParamsVerbosityRender, value);
 			}
 		}
+
+		private UE4LogVerbosity _SelectedVerbocityRenderSync = UE4LogVerbosity.Log;
+		public UE4LogVerbosity SelectedVerbocityRenderSync
+		{
+			get { return _SelectedVerbocityRenderSync; }
+			set
+			{
+				Set(ref _SelectedVerbocityRenderSync, value, "SelectedVerbocityRenderSync");
+				RegistrySaver.UpdateRegistry(RegistrySaver.RegCategoryLogParams, RegistrySaver.RegLogParamsVerbosityRenderSync, value);
+			}
+		}
 		#endregion
 
 		public Launcher()
diff --git a/Engine/Source/Programs/nDisplayLauncher/Cluster/Log/MainWindow.Log.cs b/Engine/Source/Programs/nDisplayLauncher/Cluster/Log/MainWindow.Log.cs
--- a/Engine/Source/Programs/nDisplayLauncher/Cluster/Log/MainWindow.Log.cs
+++ b/Engine/Source/Programs/nDisplayLauncher/Cluster/Log/MainWindow.Log.cs
@@ -71,6 +71,7 @@ namespace nDisplayLauncher
 			TheLauncher.SelectedVerbocityNetwork    = Verbosity;
 			TheLauncher.SelectedVerbocityNetworkMsg = Verbosity;
 			TheLauncher.SelectedVerbocityRender     = Verbosity;
+			TheLauncher.SelectedVerbocityRenderSync = Verbosity;
 			TheLauncher.SelectedVerbocityBlueprint  = Verbosity;
 		}
 	}
diff --git a/Engine/Source/Programs/nDisplayLauncher/MainWindow.xaml b/Engine/Source/Programs/nDisplayLauncher/MainWindow.xaml
--- a/Engine/Source/Programs/nDisplayLauncher/MainWindow.xaml
+++ b/Engine/Source/Programs/nDisplayLauncher/MainWindow.xaml
@@ -592,6 +592,7 @@
 						<RowDefinition Height="Auto" />
 						<RowDefinition Height="Auto" />
 						<RowDefinition Height="Auto" />
+						<RowDefinition Height="Auto" />
 						<RowDefinition Height="*" />
 					</Grid.RowDefinitions>
 					<Grid.ColumnDefinitions>
@@ -638,6 +639,9 @@
 					<Label Grid.Row="11" Grid.Column="0" Content="Render" />
 					<ComboBox Grid.Row="11" Grid.Column="1" x:Name="ctrlLogRender" SelectedItem="{Binding SelectedVerbocityRender}" ItemsSource="{Binding UE4LogVerbosityEnumTypeValues}" Template="{DynamicResource ComboBoxTemplate}" Width="Auto" HorizontalAlignment="Stretch" />
 
+					<Label Grid.Row="12" Grid.Column="0" Content="RenderSync" />
+					<ComboBox Grid.Row="12" Grid.Column="1" x:Name="ctrlLogRenderSync" SelectedItem="{Binding SelectedVerbocityRenderSync}" ItemsSource="{Binding UE4LogVerbosityEnumTypeValues}" Template="{DynamicResource ComboBoxTemplate}" Width="Auto" HorizontalAlignment="Stretch" />
+
 					<Label Grid.Row="0" Grid.Column="3" Content="For each:" />
 					<Button x:Name="ctrlBtnSetVerbosityAll" Grid.Row="1" Grid.Column="3" Content="All" Click="ctrlBtnSetVerbosityAll_Click" />
 					<Button x:Name="ctrlBtnSetVerbosityVerbose" Grid.Row="2" Grid.Column="3" Content="Verbose" Click="ctrlBtnSetVerbosityAll_Click" />
diff --git a/Engine/Source/Programs/nDisplayLauncher/Settings/RegistrySaver.cs b/Engine/Source/Programs/nDisplayLauncher/Settings/RegistrySaver.cs
--- a/Engine/Source/Programs/nDisplayLauncher/Settings/RegistrySaver.cs
+++ b/Engine/Source/Programs/nDisplayLauncher/Settings/RegistrySaver.cs
@@ -45,7 +45,7 @@ namespace nDisplayLauncher.Settings
 		public const string RegLogParamsVerbosityNetworkMsg = "VerbosityNetworkMsg";
 		public const string RegLogParamsVerbosityBlueprint  = "VerbosityBlueprint";
 		public const string RegLogParamsVerbosityRender     = "VerbosityRender";
-
+		public const string RegLogParamsVerbosityRenderSync = "VerbosityRenderSync";
 
 
 		private static string[] ReadRegistry(string key)
