Temporary warning if a config path has spaces

diff --git a/Engine/Source/Programs/nDisplayLauncher/Cluster/Launcher/MainWindow.Launcher.cs b/Engine/Source/Programs/nDisplayLauncher/Cluster/Launcher/MainWindow.Launcher.cs
--- a/Engine/Source/Programs/nDisplayLauncher/Cluster/Launcher/MainWindow.Launcher.cs
+++ b/Engine/Source/Programs/nDisplayLauncher/Cluster/Launcher/MainWindow.Launcher.cs
@@ -39,6 +39,12 @@ namespace nDisplayLauncher
 				return;
 			}
 
+			if (TheLauncher.SelectedConfig.Contains(" "))
+			{
+				System.Windows.Forms.MessageBox.Show("Paths with spaces currently not supported. Please make sure there is no spaces in the config path.");
+				return;
+			}
+
 			TheLauncher.ProcessCommand(Launcher.ClusterCommandType.RunApp);
 		}
 
@@ -129,6 +135,12 @@ namespace nDisplayLauncher
 			if (openFileDialog.ShowDialog() == System.Windows.Forms.DialogResult.OK)
 			{
 				string configPath = openFileDialog.FileName;
+				if (configPath.Contains(" "))
+				{
+					System.Windows.Forms.MessageBox.Show("Paths with spaces currently not supported");
+					return;
+				}
+
 				if (!TheLauncher.Configs.Exists(x => x == configPath))
 				{
 					TheLauncher.AddConfig(configPath);
