nDisplay: fixed UE4 mentions in the source code

diff --git a/Engine/Plugins/Runtime/WarpUtils/Source/PFMExporter/Private/PFMExporterShader.cpp b/Engine/Plugins/Runtime/WarpUtils/Source/PFMExporter/Private/PFMExporterShader.cpp
--- a/Engine/Plugins/Runtime/WarpUtils/Source/PFMExporter/Private/PFMExporterShader.cpp
+++ b/Engine/Plugins/Runtime/WarpUtils/Source/PFMExporter/Private/PFMExporterShader.cpp
@@ -81,7 +81,7 @@ public:
 	}	
 };
 
-// Implement shaders inside UE4
+// Implement shaders inside UE
 IMPLEMENT_SHADER_TYPE(, FPFMExporterVS, PFMExporterShaderFileName, TEXT("PFMExporterUV_VS"), SF_Vertex);
 IMPLEMENT_SHADER_TYPE(, FPFMExporterPS, PFMExporterShaderFileName, TEXT("PFMExporterPassthrough_PS"), SF_Pixel);
 
diff --git a/Engine/Plugins/Runtime/nDisplay/Source/DisplayCluster/Private/DisplayClusterModule.cpp b/Engine/Plugins/Runtime/nDisplay/Source/DisplayCluster/Private/DisplayClusterModule.cpp
--- a/Engine/Plugins/Runtime/nDisplay/Source/DisplayCluster/Private/DisplayClusterModule.cpp
+++ b/Engine/Plugins/Runtime/nDisplay/Source/DisplayCluster/Private/DisplayClusterModule.cpp
@@ -35,7 +35,7 @@ FDisplayClusterModule::~FDisplayClusterModule()
 	GDisplayCluster = nullptr;
 #else
 	// WORKAROUND
-	// UE4 does something like that:
+	// The engine does something like that:
 	// 1. inst1 = new FDisplayClusterModule
 	// 2. inst2 = new FDisplayClusterModule
 	// 3. delete inst1
diff --git a/Engine/Plugins/Runtime/nDisplay/Source/DisplayCluster/Private/Misc/DisplayClusterAppExit.cpp b/Engine/Plugins/Runtime/nDisplay/Source/DisplayCluster/Private/Misc/DisplayClusterAppExit.cpp
--- a/Engine/Plugins/Runtime/nDisplay/Source/DisplayCluster/Private/Misc/DisplayClusterAppExit.cpp
+++ b/Engine/Plugins/Runtime/nDisplay/Source/DisplayCluster/Private/Misc/DisplayClusterAppExit.cpp
@@ -18,9 +18,9 @@ auto FDisplayClusterAppExit::ExitTypeToStr(EExitType ExitType)
 	case EExitType::KillImmediately:
 		return TEXT("KILL");
 	case EExitType::NormalSoft:
-		return TEXT("UE4_soft");
+		return TEXT("UE_soft");
 	case EExitType::NormalForce:
-		return TEXT("UE4_force");
+		return TEXT("UE_force");
 	default:
 		return TEXT("unknown");
 	}
@@ -73,8 +73,6 @@ void FDisplayClusterAppExit::ExitApplication(EExitType ExitType, const FString&
 
 				case EExitType::NormalSoft:
 				{
-//@todo: This is workaround for exit issue - crash on exit. Need to be checked on new UE versions.
-// <ErrorMessage>Assertion failed: NumRemoved == 1 [File:D:\work\UE4.12.5.build\Engine\Source\Runtime\CoreUObject\Private\UObject\UObjectHash.cpp] [Line: 905] &nl;&nl;</ErrorMessage>
 					FProcHandle hProc = FPlatformProcess::OpenProcess(FPlatformProcess::GetCurrentProcessId());
 					FPlatformProcess::TerminateProc(hProc, true);
 					break;
diff --git a/Engine/Plugins/Runtime/nDisplay/Source/DisplayCluster/Private/Misc/DisplayClusterAppExit.h b/Engine/Plugins/Runtime/nDisplay/Source/DisplayCluster/Private/Misc/DisplayClusterAppExit.h
--- a/Engine/Plugins/Runtime/nDisplay/Source/DisplayCluster/Private/Misc/DisplayClusterAppExit.h
+++ b/Engine/Plugins/Runtime/nDisplay/Source/DisplayCluster/Private/Misc/DisplayClusterAppExit.h
@@ -15,9 +15,9 @@ public:
 	{
 		// Kills current process. No resource cleaning performed.
 		KillImmediately,
-		// UE4 based soft exit (game thread). Full resource cleaning.
+		// UE based soft exit (game thread). Full resource cleaning.
 		NormalSoft,
-		// UE4 game termination. Error window and dump file should appear after exit.
+		// UE game termination. Error window and dump file should appear after exit.
 		NormalForce
 	};
 
diff --git a/Engine/Plugins/Runtime/nDisplay/Source/DisplayCluster/Private/Network/Protocol/IDisplayClusterProtocolClusterSync.h b/Engine/Plugins/Runtime/nDisplay/Source/DisplayCluster/Private/Network/Protocol/IDisplayClusterProtocolClusterSync.h
--- a/Engine/Plugins/Runtime/nDisplay/Source/DisplayCluster/Private/Network/Protocol/IDisplayClusterProtocolClusterSync.h
+++ b/Engine/Plugins/Runtime/nDisplay/Source/DisplayCluster/Private/Network/Protocol/IDisplayClusterProtocolClusterSync.h
@@ -40,6 +40,6 @@ public:
 	// Sync events
 	virtual void GetEventsData(TArray<TSharedPtr<FDisplayClusterClusterEventJson, ESPMode::ThreadSafe>>& JsonEvents, TArray<TSharedPtr<FDisplayClusterClusterEventBinary, ESPMode::ThreadSafe>>& BinaryEvents) = 0;
 
-	// Sync native UE4 input
+	// Sync native UE input
 	virtual void GetNativeInputData(TMap<FString, FString>& NativeInputData) = 0;
 };
diff --git a/Engine/Plugins/Runtime/nDisplay/Source/DisplayCluster/Private/Render/Device/DisplayClusterDeviceBase.h b/Engine/Plugins/Runtime/nDisplay/Source/DisplayCluster/Private/Render/Device/DisplayClusterDeviceBase.h
--- a/Engine/Plugins/Runtime/nDisplay/Source/DisplayCluster/Private/Render/Device/DisplayClusterDeviceBase.h
+++ b/Engine/Plugins/Runtime/nDisplay/Source/DisplayCluster/Private/Render/Device/DisplayClusterDeviceBase.h
@@ -169,7 +169,7 @@ protected:
 	mutable TArray<FDisplayClusterRenderViewport> RenderViewports;
 	// Views per viewport (render passes)
 	uint32 ViewsAmountPerViewport = 0;
-	// UE4 main viewport
+	// UE main viewport
 	FViewport* MainViewport = nullptr;
 
 	bool bIsSceneOpen = false;
diff --git a/Engine/Plugins/Runtime/nDisplay/Source/DisplayCluster/Private/Render/DisplayClusterRenderManager.h b/Engine/Plugins/Runtime/nDisplay/Source/DisplayCluster/Private/Render/DisplayClusterRenderManager.h
--- a/Engine/Plugins/Runtime/nDisplay/Source/DisplayCluster/Private/Render/DisplayClusterRenderManager.h
+++ b/Engine/Plugins/Runtime/nDisplay/Source/DisplayCluster/Private/Render/DisplayClusterRenderManager.h
@@ -104,6 +104,6 @@ private:
 	// Internal data access synchronization
 	mutable FCriticalSection CritSecInternals;
 
-	// This flag is used to auto-focus the UE4 window once on start
+	// This flag is used to auto-focus the UE window once on start
 	bool bWasWindowFocused = false;
 };
diff --git a/Engine/Plugins/Runtime/nDisplay/Source/DisplayCluster/Private/Render/Synchronization/DisplayClusterRenderSyncPolicyNone.cpp b/Engine/Plugins/Runtime/nDisplay/Source/DisplayCluster/Private/Render/Synchronization/DisplayClusterRenderSyncPolicyNone.cpp
--- a/Engine/Plugins/Runtime/nDisplay/Source/DisplayCluster/Private/Render/Synchronization/DisplayClusterRenderSyncPolicyNone.cpp
+++ b/Engine/Plugins/Runtime/nDisplay/Source/DisplayCluster/Private/Render/Synchronization/DisplayClusterRenderSyncPolicyNone.cpp
@@ -16,7 +16,7 @@ FDisplayClusterRenderSyncPolicyNone::~FDisplayClusterRenderSyncPolicyNone()
 
 bool FDisplayClusterRenderSyncPolicyNone::SynchronizeClusterRendering(int32& InOutSyncInterval)
 {
-	UE_LOG(LogDisplayClusterRenderSync, VeryVerbose, TEXT("Sync policy is 'None'. UE4 will swap the buffers with swap interval 0"));
+	UE_LOG(LogDisplayClusterRenderSync, VeryVerbose, TEXT("Sync policy is 'None', swap interval 0"));
 
 	// Override sync interval with 0 to show a frame ASAP. We don't care about tearing in this policy.
 	InOutSyncInterval = 0;
diff --git a/Engine/Plugins/Runtime/nDisplay/Source/DisplayCluster/Public/Misc/DisplayClusterHelpers.h b/Engine/Plugins/Runtime/nDisplay/Source/DisplayCluster/Public/Misc/DisplayClusterHelpers.h
--- a/Engine/Plugins/Runtime/nDisplay/Source/DisplayCluster/Public/Misc/DisplayClusterHelpers.h
+++ b/Engine/Plugins/Runtime/nDisplay/Source/DisplayCluster/Public/Misc/DisplayClusterHelpers.h
@@ -577,7 +577,7 @@ namespace DisplayClusterHelpers
 				FPlane(ma, mb, mc, me),
 				FPlane(0,  0,  md,  0));
 
-			// Invert Z-axis (UE4 uses Z-inverted LHS)
+			// Invert Z-axis (UE uses Z-inverted LHS)
 			static const FMatrix flipZ = FMatrix(
 				FPlane(1, 0, 0, 0),
 				FPlane(0, 1, 0, 0),
diff --git a/Engine/Plugins/Runtime/nDisplay/Source/DisplayCluster/Public/Render/IDisplayClusterRenderManager.h b/Engine/Plugins/Runtime/nDisplay/Source/DisplayCluster/Public/Render/IDisplayClusterRenderManager.h
--- a/Engine/Plugins/Runtime/nDisplay/Source/DisplayCluster/Public/Render/IDisplayClusterRenderManager.h
+++ b/Engine/Plugins/Runtime/nDisplay/Source/DisplayCluster/Public/Render/IDisplayClusterRenderManager.h
@@ -275,7 +275,7 @@ public:
 	* Configuration of interpupillary (interocular) distance
 	*
 	* @param CameraID    - Camera ID to modify
-	* @param EyeDistance - distance between eyes (UE4 units).
+	* @param EyeDistance - distance between eyes (engine units).
 	*/
 	UE_DEPRECATED(4.26, "This function has been moved to UDisplayClusterCameraComponent")
 	virtual void  SetInterpupillaryDistance(const FString& CameraID, float EyeDistance)
@@ -286,7 +286,7 @@ public:
 	*
 	* @param CameraID    - Camera ID to get data
 	*
-	* @return - distance between eyes (UE4 units)
+	* @return - distance between eyes (engine units)
 	*/
 	UE_DEPRECATED(4.26, "This function has been moved to UDisplayClusterCameraComponent")
 	virtual float GetInterpupillaryDistance(const FString& CameraID) const
diff --git a/Engine/Plugins/Runtime/nDisplay/Source/DisplayClusterConfiguration/Private/Formats/Text/DisplayClusterConfigurationTextStrings.h b/Engine/Plugins/Runtime/nDisplay/Source/DisplayClusterConfiguration/Private/Formats/Text/DisplayClusterConfigurationTextStrings.h
--- a/Engine/Plugins/Runtime/nDisplay/Source/DisplayClusterConfiguration/Private/Formats/Text/DisplayClusterConfigurationTextStrings.h
+++ b/Engine/Plugins/Runtime/nDisplay/Source/DisplayClusterConfiguration/Private/Formats/Text/DisplayClusterConfigurationTextStrings.h
@@ -158,7 +158,7 @@ namespace DisplayClusterConfigurationTextStrings
 			{
 				static constexpr auto Header                = TEXT("[general]");
 				static constexpr auto SwapSyncPolicy        = TEXT("swap_sync_policy");
-				static constexpr auto UnrealInputSyncPolicy = TEXT("ue4_input_sync_policy");
+				static constexpr auto UnrealInputSyncPolicy = TEXT("ue_input_sync_policy");
 			}
 
 			namespace nvidia
diff --git a/Engine/Plugins/Runtime/nDisplay/Source/DisplayClusterConfiguration/Private/Formats/Text/DisplayClusterConfigurationTextTypes.h b/Engine/Plugins/Runtime/nDisplay/Source/DisplayClusterConfiguration/Private/Formats/Text/DisplayClusterConfigurationTextTypes.h
--- a/Engine/Plugins/Runtime/nDisplay/Source/DisplayClusterConfiguration/Private/Formats/Text/DisplayClusterConfigurationTextTypes.h
+++ b/Engine/Plugins/Runtime/nDisplay/Source/DisplayClusterConfiguration/Private/Formats/Text/DisplayClusterConfigurationTextTypes.h
@@ -87,7 +87,7 @@ struct FDisplayClusterConfigurationTextViewport : public FDisplayClusterConfigur
 	FIntPoint Size = FIntPoint::ZeroValue;
 	float     BufferRatio = 1.f;
 	int       GPUIndex = -1;                 // Force custom mgpu index for this viewport {view->bOverrideGPUMask=true; view->GPUMask=GPUIndex; }
-	bool      AllowCrossGPUTransfer = true;  // Control UE4 viewport mgpu transfer View->bAllowCrossGPUTransfer
+	bool      AllowCrossGPUTransfer = true;  // Control UE viewport mgpu transfer View->bAllowCrossGPUTransfer
 	bool      IsShared = false;              // Share this viewport for all (scene context textures, backbuffer)
 
 	virtual FString ToString() const override;
diff --git a/Engine/Plugins/Runtime/nDisplay/Source/DisplayClusterConfiguration/Public/DisplayClusterConfigurationTypes.h b/Engine/Plugins/Runtime/nDisplay/Source/DisplayClusterConfiguration/Public/DisplayClusterConfigurationTypes.h
--- a/Engine/Plugins/Runtime/nDisplay/Source/DisplayClusterConfiguration/Public/DisplayClusterConfigurationTypes.h
+++ b/Engine/Plugins/Runtime/nDisplay/Source/DisplayClusterConfiguration/Public/DisplayClusterConfigurationTypes.h
@@ -26,7 +26,7 @@ enum class EDisplayClusterConfigurationKeyboardReflectionType : uint8
 	None     UMETA(DisplayName = "No reflection"),
 	nDisplay UMETA(DisplayName = "nDisplay buttons only"),
 	Core     UMETA(DisplayName = "UE core keyboard events only"),
-	All      UMETA(DisplayName = "Both nDisplay and UE4 core events")
+	All      UMETA(DisplayName = "Both nDisplay and UE core events")
 };
 
 UENUM()
diff --git a/Engine/Plugins/Runtime/nDisplay/Source/DisplayClusterInput/Private/Blueprints/DisplayClusterInputBlueprintAPIImpl.h b/Engine/Plugins/Runtime/nDisplay/Source/DisplayClusterInput/Private/Blueprints/DisplayClusterInputBlueprintAPIImpl.h
--- a/Engine/Plugins/Runtime/nDisplay/Source/DisplayClusterInput/Private/Blueprints/DisplayClusterInputBlueprintAPIImpl.h
+++ b/Engine/Plugins/Runtime/nDisplay/Source/DisplayClusterInput/Private/Blueprints/DisplayClusterInputBlueprintAPIImpl.h
@@ -37,7 +37,7 @@ public:
 	virtual bool BindVrpnChannel(const FString& VrpnDeviceId, const int32 VrpnChannel, const FKey Target) override;
 
 	/**
-	* Create new bind from vrpn keyboard device key to UE4 target by user friendly target name
+	* Create new bind from vrpn keyboard device key to a target by user friendly target name
 	*
 	* @return True if success
 	*/
diff --git a/Engine/Plugins/Runtime/nDisplay/Source/DisplayClusterInput/Private/Controller/DisplayClusterInputControllerAnalog.cpp b/Engine/Plugins/Runtime/nDisplay/Source/DisplayClusterInput/Private/Controller/DisplayClusterInputControllerAnalog.cpp
--- a/Engine/Plugins/Runtime/nDisplay/Source/DisplayClusterInput/Private/Controller/DisplayClusterInputControllerAnalog.cpp
+++ b/Engine/Plugins/Runtime/nDisplay/Source/DisplayClusterInput/Private/Controller/DisplayClusterInputControllerAnalog.cpp
@@ -13,7 +13,7 @@
 #define LOCTEXT_NAMESPACE "DisplayClusterInput"
 
 
-// Add vrpn analog to UE4 global name-space
+// Add vrpn analog to the engine global name-space
 void FAnalogController::Initialize()
 {
 	static const FName nDisplayClusterInputCategoryName(TEXT("nDisplayAnalogs"));
diff --git a/Engine/Plugins/Runtime/nDisplay/Source/DisplayClusterInput/Private/Controller/DisplayClusterInputControllerBase.h b/Engine/Plugins/Runtime/nDisplay/Source/DisplayClusterInput/Private/Controller/DisplayClusterInputControllerBase.h
--- a/Engine/Plugins/Runtime/nDisplay/Source/DisplayClusterInput/Private/Controller/DisplayClusterInputControllerBase.h
+++ b/Engine/Plugins/Runtime/nDisplay/Source/DisplayClusterInput/Private/Controller/DisplayClusterInputControllerBase.h
@@ -170,7 +170,7 @@ public:
 		return true;
 	}
 
-	// Send vrpn channels data to UE4 core
+	// Send vrpn channels data to the engine core
 	void UpdateEvents(const double CurrentTime, FGenericApplicationMessageHandler* MessageHandler, int32 ControllerId)
 	{
 		for (auto& DeviceIt : BindMap)
diff --git a/Engine/Plugins/Runtime/nDisplay/Source/DisplayClusterInput/Private/Controller/DisplayClusterInputControllerKeyboard.cpp b/Engine/Plugins/Runtime/nDisplay/Source/DisplayClusterInput/Private/Controller/DisplayClusterInputControllerKeyboard.cpp
--- a/Engine/Plugins/Runtime/nDisplay/Source/DisplayClusterInput/Private/Controller/DisplayClusterInputControllerKeyboard.cpp
+++ b/Engine/Plugins/Runtime/nDisplay/Source/DisplayClusterInput/Private/Controller/DisplayClusterInputControllerKeyboard.cpp
@@ -282,7 +282,7 @@ void FKeyboardController::ConnectKey(FChannelBinds& KeyboardData, uint32 VrpnCha
 #if PLATFORM_WINDOWS
 #include "Windows/AllowWindowsPlatformTypes.h"
 
-// Support reflect vrpn to ue4\nDisplay:
+// Support vrpn reflection
 void FKeyboardController::ReflectKeyboard(const FString& VrpnDeviceId, EDisplayClusterInputKeyboardReflectionMode ReflectionMode)
 {
 	switch (ReflectionMode)
diff --git a/Engine/Plugins/Runtime/nDisplay/Source/DisplayClusterInput/Private/Controller/DisplayClusterInputControllerKeyboard.h b/Engine/Plugins/Runtime/nDisplay/Source/DisplayClusterInput/Private/Controller/DisplayClusterInputControllerKeyboard.h
--- a/Engine/Plugins/Runtime/nDisplay/Source/DisplayClusterInput/Private/Controller/DisplayClusterInputControllerKeyboard.h
+++ b/Engine/Plugins/Runtime/nDisplay/Source/DisplayClusterInput/Private/Controller/DisplayClusterInputControllerKeyboard.h
@@ -16,7 +16,7 @@ public:
 	{ }
 
 public:
-	// Register nDisplay second keyboard keys in UE4 FKey namespace
+	// Register nDisplay second keyboard keys in the engine FKey namespace
 	virtual void Initialize() override;
 
 	// Delegated events for DisplayClusterInput vrpn connect
@@ -24,7 +24,7 @@ public:
 	virtual void ProcessEndSession() override;
 	virtual void ProcessPreTick() override;
 
-	// Reflect vrpn keyboard to UE4
+	// Reflect vrpn keyboard to the engine
 	void ReflectKeyboard(const FString& VrpnDeviceId, EDisplayClusterInputKeyboardReflectionMode ReflectionMode);
 
 
diff --git a/Engine/Plugins/Runtime/nDisplay/Source/DisplayClusterInput/Private/DisplayClusterInput.h b/Engine/Plugins/Runtime/nDisplay/Source/DisplayClusterInput/Private/DisplayClusterInput.h
--- a/Engine/Plugins/Runtime/nDisplay/Source/DisplayClusterInput/Private/DisplayClusterInput.h
+++ b/Engine/Plugins/Runtime/nDisplay/Source/DisplayClusterInput/Private/DisplayClusterInput.h
@@ -38,5 +38,5 @@ private:
 	FDisplayClusterInputModule&                   InputModuleAPI;
 	TSharedPtr<FGenericApplicationMessageHandler> MessageHandler;
 	
-	int32 UnrealControllerIndex;  // Local UE4 player index (multi player purpose)
+	int32 UnrealControllerIndex;  // Local player index (multi player purpose)
 };
diff --git a/Engine/Plugins/Runtime/nDisplay/Source/DisplayClusterInput/Private/DisplayClusterInputModule.cpp b/Engine/Plugins/Runtime/nDisplay/Source/DisplayClusterInput/Private/DisplayClusterInputModule.cpp
--- a/Engine/Plugins/Runtime/nDisplay/Source/DisplayClusterInput/Private/DisplayClusterInputModule.cpp
+++ b/Engine/Plugins/Runtime/nDisplay/Source/DisplayClusterInput/Private/DisplayClusterInputModule.cpp
@@ -49,7 +49,7 @@ void FDisplayClusterInputModule::SendControllerEvents(const TSharedPtr<FGenericA
 {
 	const double CurrentTime = FPlatformTime::Seconds();
 	
-	// Send all changes for button and analog to UE4 core
+	// Send all changes for button and analog to the engine core
 	ButtonController.UpdateEvents(CurrentTime, MessageHandler.Get(), UnrealControllerIndex);
 	AnalogController.UpdateEvents(CurrentTime, MessageHandler.Get(), UnrealControllerIndex);
 	KeyboardController.UpdateEvents(CurrentTime, MessageHandler.Get(), UnrealControllerIndex);
@@ -105,7 +105,7 @@ bool FDisplayClusterInputModule::BindVrpnChannel(const FString& VrpnDeviceId, ui
 	return false;
 }
 
-// Bind all keyboard keys to ue4 (default keyboard and|or nDisplay second keyboard namespaces)
+// Bind all keyboard keys to the engine (default keyboard and|or nDisplay second keyboard namespaces)
 bool FDisplayClusterInputModule::SetVrpnKeyboardReflectionMode(const FString& VrpnDeviceId, EDisplayClusterInputKeyboardReflectionMode ReflectionMode)
 {
 	if (IsSessionStarted())
diff --git a/Engine/Plugins/Runtime/nDisplay/Source/DisplayClusterInput/Private/DisplayClusterInputModule.h b/Engine/Plugins/Runtime/nDisplay/Source/DisplayClusterInput/Private/DisplayClusterInputModule.h
--- a/Engine/Plugins/Runtime/nDisplay/Source/DisplayClusterInput/Private/DisplayClusterInputModule.h
+++ b/Engine/Plugins/Runtime/nDisplay/Source/DisplayClusterInput/Private/DisplayClusterInputModule.h
@@ -26,9 +26,9 @@ public:
 	virtual TSharedPtr< class IInputDevice > CreateInputDevice( const TSharedRef< FGenericApplicationMessageHandler >& InMessageHandler ) override;
 
 public:
-	// Add new VRPN device input bind from channel to UE4 target (type auto detected by device ID)
+	// Add new VRPN device input bind from channel to a target (type auto detected by device ID)
 	virtual bool BindVrpnChannel(const FString& VrpnDeviceId, uint32 VrpnChannel, const FString& BindTargetName) override;
-	// Bind all keyboard keys to ue4 (default keyboard and|or nDisplay second keyboard namespaces)
+	// Bind all keyboard keys (default keyboard and|or nDisplay second keyboard namespaces)
 	virtual bool SetVrpnKeyboardReflectionMode(const FString& VrpnDeviceId, EDisplayClusterInputKeyboardReflectionMode ReflectionMode) override;
 	// Apply all delayed vrpn bindings
 	void UpdateVrpnBindings();
diff --git a/Engine/Plugins/Runtime/nDisplay/Source/DisplayClusterInput/Private/Misc/DisplayClusterInputHelpers.h b/Engine/Plugins/Runtime/nDisplay/Source/DisplayClusterInput/Private/Misc/DisplayClusterInputHelpers.h
--- a/Engine/Plugins/Runtime/nDisplay/Source/DisplayClusterInput/Private/Misc/DisplayClusterInputHelpers.h
+++ b/Engine/Plugins/Runtime/nDisplay/Source/DisplayClusterInput/Private/Misc/DisplayClusterInputHelpers.h
@@ -12,7 +12,7 @@ class DisplayClusterInputHelpers
 {
 public:
 	/** 
-	 * Converts keyboard button name (defined inside UE4 core) to a scancode (equals to a keyboard vrpn channel)
+	 * Converts keyboard button name (defined inside the engine core) to a scancode (equals to a keyboard vrpn channel)
 	 */
 	static bool KeyNameToVrpnScancode(const FString& KeyName, int32& VrpnScanCode);
 };
diff --git a/Engine/Plugins/Runtime/nDisplay/Source/DisplayClusterInput/Private/State/DisplayClusterInputStateAnalog.h b/Engine/Plugins/Runtime/nDisplay/Source/DisplayClusterInput/Private/State/DisplayClusterInputStateAnalog.h
--- a/Engine/Plugins/Runtime/nDisplay/Source/DisplayClusterInput/Private/State/DisplayClusterInputStateAnalog.h
+++ b/Engine/Plugins/Runtime/nDisplay/Source/DisplayClusterInput/Private/State/DisplayClusterInputStateAnalog.h
@@ -24,7 +24,7 @@ public:
 		BindKeys.Empty();
 	}
 
-	// Add new UE4 target to binds array
+	// Add new target to binds array
 	bool BindTarget(const FName& NewKey)
 	{
 		if (BindKeys.Contains(NewKey))
@@ -55,15 +55,15 @@ public:
 		AnalogState = AnalogNextState;
 	}
 
-	// Send all changes to UE4 core
+	// Send all changes to the engine core
 	void UpdateEvents(FGenericApplicationMessageHandler* MessageHandler, int32 ControllerId, double CurrentTime);
 
 private:
-	// Send axis event to UE4 core
+	// Send axis event to the engine core
 	bool OnAnalogChanges(FGenericApplicationMessageHandler* MessageHandler, int32 ControllerId);
 
 private:
-	// UE4 bind targets for this channel
+	// Binding targets for this channel
 	TArray<FName> BindKeys;
 	// Analog value [0.f,1.f]
 	float AnalogState;
diff --git a/Engine/Plugins/Runtime/nDisplay/Source/DisplayClusterInput/Private/State/DisplayClusterInputStateButton.h b/Engine/Plugins/Runtime/nDisplay/Source/DisplayClusterInput/Private/State/DisplayClusterInputStateButton.h
--- a/Engine/Plugins/Runtime/nDisplay/Source/DisplayClusterInput/Private/State/DisplayClusterInputStateButton.h
+++ b/Engine/Plugins/Runtime/nDisplay/Source/DisplayClusterInput/Private/State/DisplayClusterInputStateButton.h
@@ -19,13 +19,13 @@ public:
 		NextRepeatTime( 0.0 )
 	{ }
 
-	// Clear all binds to UE4
+	// Clear all bindings
 	void Reset()
 	{
 		BindKeys.Empty();
 	}
 
-	// Add new UE4 target to binds array
+	// Add new target to bindings array
 	bool BindTarget(const FName& NewKey)
 	{
 		if (BindKeys.Contains(NewKey))
@@ -60,11 +60,11 @@ public:
 		bIsPressedState  = bIsPressedNextState;
 	}
 
-	// Send all changes to UE4 core
+	// Send all changes to the engine core
 	void UpdateEvents(FGenericApplicationMessageHandler* MessageHandler, int32 ControllerId, double CurrentTime);
 
 private:
-	// Send button events to UE4 core
+	// Send button events to the engine core
 	bool OnButtonPressed(FGenericApplicationMessageHandler* MessageHandler, int32 ControllerId, bool bIsRepeat);
 	bool OnButtonReleased(FGenericApplicationMessageHandler* MessageHandler, int32 ControllerId, bool bIsRepeat);
 
@@ -73,7 +73,7 @@ private:
 	static float ButtonRepeatDelay;
 
 private:
-	// UE4 bind targets for this channel
+	// Bindings
 	TArray<FName> BindKeys;
 	// Whether we're pressed or not.  While pressed, we will generate repeat presses on a timer
 	bool bIsPressedState;
diff --git a/Engine/Plugins/Runtime/nDisplay/Source/DisplayClusterInput/Private/State/DisplayClusterInputStateKeyboard.cpp b/Engine/Plugins/Runtime/nDisplay/Source/DisplayClusterInput/Private/State/DisplayClusterInputStateKeyboard.cpp
--- a/Engine/Plugins/Runtime/nDisplay/Source/DisplayClusterInput/Private/State/DisplayClusterInputStateKeyboard.cpp
+++ b/Engine/Plugins/Runtime/nDisplay/Source/DisplayClusterInput/Private/State/DisplayClusterInputStateKeyboard.cpp
@@ -5,7 +5,7 @@
 
 
 /*
- * Define global UE4 names for DisplayCluster vrpn keyboard buttons:
+ * Define global engine names for DisplayCluster vrpn keyboard buttons:
  */
 
 const FKey FKeyboardKey::Zero("nDisplayKeyboardZero");
diff --git a/Engine/Plugins/Runtime/nDisplay/Source/DisplayClusterInput/Private/State/DisplayClusterInputStateTracker.h b/Engine/Plugins/Runtime/nDisplay/Source/DisplayClusterInput/Private/State/DisplayClusterInputStateTracker.h
--- a/Engine/Plugins/Runtime/nDisplay/Source/DisplayClusterInput/Private/State/DisplayClusterInputStateTracker.h
+++ b/Engine/Plugins/Runtime/nDisplay/Source/DisplayClusterInput/Private/State/DisplayClusterInputStateTracker.h
@@ -16,7 +16,7 @@ public:
 		PositionNext(0, 0, 0)
 	{ }
 
-	// Clear all binds to UE4
+	// Clear all bindings
 	void Reset()
 	{
 		Trackers.Empty();
@@ -34,7 +34,7 @@ public:
 		return Trackers.Find(NewTrackerKey);
 	}
 
-	// Add new UE4 target to tracker array
+	// Add a new target to tracker array
 	bool BindTarget(const EControllerHand NewTrackerKey)
 	{
 		if (Trackers.Contains(NewTrackerKey))
@@ -69,7 +69,7 @@ public:
 	}
 
 private:
-	// List of all UE4 trackers binded to this vrpn channel
+	// List of all trackers binded to this vrpn channel
 	TArray<EControllerHand> Trackers;
 
 	// Current tracker state
diff --git a/Engine/Plugins/Runtime/nDisplay/Source/DisplayClusterInput/Public/Blueprints/IDisplayClusterInputBlueprintAPI.h b/Engine/Plugins/Runtime/nDisplay/Source/DisplayClusterInput/Public/Blueprints/IDisplayClusterInputBlueprintAPI.h
--- a/Engine/Plugins/Runtime/nDisplay/Source/DisplayClusterInput/Public/Blueprints/IDisplayClusterInputBlueprintAPI.h
+++ b/Engine/Plugins/Runtime/nDisplay/Source/DisplayClusterInput/Public/Blueprints/IDisplayClusterInputBlueprintAPI.h
@@ -43,7 +43,7 @@ public:
 	virtual bool BindVrpnChannel(const FString& VrpnDeviceId, const int32 VrpnChannel, const FKey Target) = 0;
 
 	/**
-	* Create new bind from vrpn keyboard device key to UE4 target by user friendly target name
+	* Create new bind from vrpn keyboard device key to a target by user friendly target name
 	*
 	* @return True if success
 	*/
diff --git a/Engine/Plugins/Runtime/nDisplay/Source/DisplayClusterInput/Public/DisplayClusterInputTypes.h b/Engine/Plugins/Runtime/nDisplay/Source/DisplayClusterInput/Public/DisplayClusterInputTypes.h
--- a/Engine/Plugins/Runtime/nDisplay/Source/DisplayClusterInput/Public/DisplayClusterInputTypes.h
+++ b/Engine/Plugins/Runtime/nDisplay/Source/DisplayClusterInput/Public/DisplayClusterInputTypes.h
@@ -15,11 +15,11 @@ enum class EDisplayClusterInputKeyboardReflectionMode : uint8
 	None     UMETA(DisplayName = "No reflection"),
 	nDisplay UMETA(DisplayName = "nDisplay buttons only"),
 	Core     UMETA(DisplayName = "UE core keyboard events"),
-	All      UMETA(DisplayName = "Both nDisplay and UE4 core")
+	All      UMETA(DisplayName = "Both nDisplay and UE core")
 };
 
 
-// Binding description. Maps VRPN device channel to an UE4 FKey
+// Binding description. Maps VRPN device channel to a FKey
 USTRUCT(BlueprintType)
 struct FDisplayClusterInputBinding
 {
diff --git a/Engine/Plugins/Runtime/nDisplay/Source/DisplayClusterInput/Public/IDisplayClusterInputModule.h b/Engine/Plugins/Runtime/nDisplay/Source/DisplayClusterInput/Public/IDisplayClusterInputModule.h
--- a/Engine/Plugins/Runtime/nDisplay/Source/DisplayClusterInput/Public/IDisplayClusterInputModule.h
+++ b/Engine/Plugins/Runtime/nDisplay/Source/DisplayClusterInput/Public/IDisplayClusterInputModule.h
@@ -38,14 +38,14 @@ public:
 	}
 
 	/**
-	* Create new bind from vrpn device channel to UE4 target by user friendly target name
+	* Create new bind from vrpn device channel to a target by user friendly target name
 	*
 	* @return True, if bind created successfully
 	*/	
 	virtual bool BindVrpnChannel(const FString& VrpnDeviceId, uint32 VrpnChannel, const FString& BindTargetName) = 0;
 
 	/*
-	 * Bind all keyboard keys to ue4 (default keyboard and|or nDisplay second keyboard namespaces)
+	 * Bind all keyboard keys (default keyboard and|or nDisplay second keyboard namespaces)
 	 *
 	 * @return true, if vrpn keyboard name is valid and op success
 	 */
diff --git a/Engine/Plugins/Runtime/nDisplay/Source/DisplayClusterProjection/Private/Policy/Camera/DisplayClusterProjectionCameraPolicy.cpp b/Engine/Plugins/Runtime/nDisplay/Source/DisplayClusterProjection/Private/Policy/Camera/DisplayClusterProjectionCameraPolicy.cpp
--- a/Engine/Plugins/Runtime/nDisplay/Source/DisplayClusterProjection/Private/Policy/Camera/DisplayClusterProjectionCameraPolicy.cpp
+++ b/Engine/Plugins/Runtime/nDisplay/Source/DisplayClusterProjection/Private/Policy/Camera/DisplayClusterProjectionCameraPolicy.cpp
@@ -70,7 +70,7 @@ bool FDisplayClusterProjectionCameraPolicy::CalculateView(const uint32 ViewIdx,
 		InOutViewLocation = AssignedCamera->GetComponentLocation();
 		InOutViewRotation = AssignedCamera->GetComponentRotation();
 	}
-	// Otherwise default UE4 camera is used
+	// Otherwise default UE camera is used
 	else
 	{
 		if (World)
diff --git a/Engine/Plugins/Runtime/nDisplay/Source/DisplayClusterProjection/Private/Policy/VIOSO/DisplayClusterProjectionVIOSOPolicy.cpp b/Engine/Plugins/Runtime/nDisplay/Source/DisplayClusterProjection/Private/Policy/VIOSO/DisplayClusterProjectionVIOSOPolicy.cpp
--- a/Engine/Plugins/Runtime/nDisplay/Source/DisplayClusterProjection/Private/Policy/VIOSO/DisplayClusterProjectionVIOSOPolicy.cpp
+++ b/Engine/Plugins/Runtime/nDisplay/Source/DisplayClusterProjection/Private/Policy/VIOSO/DisplayClusterProjectionVIOSOPolicy.cpp
@@ -466,10 +466,10 @@ bool FDisplayClusterProjectionVIOSOPolicy::FViewData::InitializeVIOSO(FRHITextur
 #ifdef VIOSO_USE_GRAPHICS_API_D3D11
 	case ERenderDevice::D3D11:
 	{
-		auto UE4D3DDevice = static_cast<ID3D11Device*>(GDynamicRHI->RHIGetNativeDevice());
+		auto UED3DDevice = static_cast<ID3D11Device*>(GDynamicRHI->RHIGetNativeDevice());
 
 		FD3D11ContextHelper D3D11ContextHelper;
-		if (D3D11ContextHelper.AssignD3D11RenderTarget(RenderTargetTexture) && Warper.Initialize(UE4D3DDevice, InConfigData))
+		if (D3D11ContextHelper.AssignD3D11RenderTarget(RenderTargetTexture) && Warper.Initialize(UED3DDevice, InConfigData))
 		{
 			return true;
 		}
diff --git a/Engine/Plugins/Runtime/nDisplay/Source/DisplayClusterProjection/Private/Policy/VIOSO/DisplayClusterProjectionVIOSOWarper.cpp b/Engine/Plugins/Runtime/nDisplay/Source/DisplayClusterProjection/Private/Policy/VIOSO/DisplayClusterProjectionVIOSOWarper.cpp
--- a/Engine/Plugins/Runtime/nDisplay/Source/DisplayClusterProjection/Private/Policy/VIOSO/DisplayClusterProjectionVIOSOWarper.cpp
+++ b/Engine/Plugins/Runtime/nDisplay/Source/DisplayClusterProjection/Private/Policy/VIOSO/DisplayClusterProjectionVIOSOWarper.cpp
@@ -113,7 +113,7 @@ bool FViosoWarper::CalculateViewProjection(FVector& InOutViewLocation, FRotator&
 
 	if (VWB_ERROR_NONE == FLibVIOSO::GetViewClip(pWarper, &InViosoEyeLocation.X, &InViosoEyeEulerRotation.X, &ViewMatrix[0], &ViewClip[0]) && IsViewClipValid())
 	{
-		// Convert to UE4 coordinate system
+		// Convert to UE coordinate system
 		FVector OutViosoEyeLocation = GetViosoViewLocation();
 		FVector OutViewLocation = FromViosoLocation(OutViosoEyeLocation, WorldToMeters);
 
@@ -139,26 +139,26 @@ bool FViosoWarper::IsViewClipValid() const
 	return true;
 }
 
-const float UE4ViosoUnitsInMeter = 1.f;
+const float UEViosoUnitsInMeter = 1.f;
 
 FVector FViosoWarper::ToViosoLocation(const FVector& InPos, const float WorldToMeters) const
 {
-	float Scale = (UE4ViosoUnitsInMeter / WorldToMeters);
+	float Scale = (UEViosoUnitsInMeter / WorldToMeters);
 
-	// UE4 coords(x,y,z) = fw,right,up
-	// UE4+VIOSO(X,Y,Z) = RIGHT, UP, FW = (y,z,x)
+	// UE coords(x,y,z) = fw,right,up
+	// UE+VIOSO(X,Y,Z) = RIGHT, UP, FW = (y,z,x)
 	FVector ViosoLocation(-InPos.Y, -InPos.Z, InPos.X);
 	return ViosoLocation * Scale;
 }
 
 FVector FViosoWarper::FromViosoLocation(const FVector& InPos, const float WorldToMeters) const
 {
-	float Scale = (WorldToMeters / UE4ViosoUnitsInMeter);
+	float Scale = (WorldToMeters / UEViosoUnitsInMeter);
 
-	// UE4+VIOSO(X,Y,Z) = RIGHT, UP, FW
-	// UE4 coords(x,y,z) = fw,right,up = (Z,X,Y)
-	FVector UE4Location(InPos.Z, -InPos.X, -InPos.Y);
-	return UE4Location * Scale;
+	// UE+VIOSO(X,Y,Z) = RIGHT, UP, FW
+	// UE coords(x,y,z) = fw,right,up = (Z,X,Y)
+	FVector UELocation(InPos.Z, -InPos.X, -InPos.Y);
+	return UELocation * Scale;
 }
 
 FVector  FViosoWarper::ToViosoEulerRotation(const FRotator& InRotation) const
diff --git a/Engine/Plugins/Runtime/nDisplay/Source/MPCDI/Private/MPCDIStrings.h b/Engine/Plugins/Runtime/nDisplay/Source/MPCDI/Private/MPCDIStrings.h
--- a/Engine/Plugins/Runtime/nDisplay/Source/MPCDI/Private/MPCDIStrings.h
+++ b/Engine/Plugins/Runtime/nDisplay/Source/MPCDI/Private/MPCDIStrings.h
@@ -17,7 +17,7 @@ namespace DisplayClusterMPCDIStrings
 
 		static constexpr auto FilePFM       = TEXT("pfm");
 		static constexpr auto WorldScale    = TEXT("scale");
-		static constexpr auto UseUnrealAxis = TEXT("ue4space");
+		static constexpr auto UseUnrealAxis = TEXT("ue_space");
 
 		static constexpr auto FileAlpha  = TEXT("alpha");
 		static constexpr auto AlphaGamma = TEXT("alpha_gamma");
diff --git a/Engine/Plugins/Runtime/nDisplay/Source/MPCDI/Public/IMPCDI.h b/Engine/Plugins/Runtime/nDisplay/Source/MPCDI/Public/IMPCDI.h
--- a/Engine/Plugins/Runtime/nDisplay/Source/MPCDI/Public/IMPCDI.h
+++ b/Engine/Plugins/Runtime/nDisplay/Source/MPCDI/Public/IMPCDI.h
@@ -300,7 +300,7 @@ public:
 	*
 	* @param InRegionLocator - region locator
 	* @param PFMFile - full path to PFM file
-	* @param WorldScale - world scale multiplier to UE4 game (value=100 from meters->sm)
+	* @param WorldScale - world scale multiplier (value=100 from meters->cm)
 	* @param bIsUnrealGameSpace - let=true, if you dont want to apply axis conversion from mpcdi space
 	*
 	* @return - true if success
@@ -314,7 +314,7 @@ public:
 	* @param PFMPoints - PFM points, row-by-row
 	* @param DimW - Row width
 	* @param DimH - Rows count
-	* @param WorldScale - world scale multiplier to UE4 game (value=100 from meters->sm)
+	* @param WorldScale - world scale multiplier (value=100 from meters->cm)
 	* @param bIsUnrealGameSpace - let=true, if you dont want to apply axis conversion from mpcdi space
 	*
 	* @return - true if success
diff --git a/Engine/Plugins/Runtime/nDisplay/Source/OutputRemap/Private/OutputRemapShader.cpp b/Engine/Plugins/Runtime/nDisplay/Source/OutputRemap/Private/OutputRemapShader.cpp
--- a/Engine/Plugins/Runtime/nDisplay/Source/OutputRemap/Private/OutputRemapShader.cpp
+++ b/Engine/Plugins/Runtime/nDisplay/Source/OutputRemap/Private/OutputRemapShader.cpp
@@ -95,7 +95,7 @@ private:
 	LAYOUT_FIELD(FShaderResourceParameter, PostprocessInputParameterSampler0)
 };
 
-// Implement shaders inside UE4
+// Implement shaders inside UE
 IMPLEMENT_SHADER_TYPE(, FOutputRemapVS, OutputRemapShaderFileName, TEXT("OutputRemap_VS"), SF_Vertex);
 IMPLEMENT_SHADER_TYPE(, FOutputRemapPS, OutputRemapShaderFileName, TEXT("OutputRemap_PS"), SF_Pixel);
 
diff --git a/Engine/Source/Programs/nDisplayLauncher/Cluster/Launcher/Launcher.cs b/Engine/Source/Programs/nDisplayLauncher/Cluster/Launcher/Launcher.cs
--- a/Engine/Source/Programs/nDisplayLauncher/Cluster/Launcher/Launcher.cs
+++ b/Engine/Source/Programs/nDisplayLauncher/Cluster/Launcher/Launcher.cs
@@ -213,7 +213,7 @@ namespace nDisplayLauncher.Cluster
 			}
 		}
 
-		public static UE4LogVerbosity UE4LogVerbosity_FromString(string From, UE4LogVerbosity Default)
+		public static UELogVerbosity UELogVerbosity_FromString(string From, UELogVerbosity Default)
 		{			
 			if (string.IsNullOrEmpty(From))
 			{
@@ -224,46 +224,46 @@ namespace nDisplayLauncher.Cluster
 
 			if (string.Compare(From, "All", true) == 0)
 			{
-				return UE4LogVerbosity.All;
+				return UELogVerbosity.All;
 			}
 			else if (string.Compare(From, "Verbose", true) == 0)
 			{
-				return UE4LogVerbosity.Verbose;
+				return UELogVerbosity.Verbose;
 			}
 			else if (string.Compare(From, "Log", true) == 0)
 			{
-				return UE4LogVerbosity.Log;
+				return UELogVerbosity.Log;
 			}
 			else if (string.Compare(From, "Display", true) == 0)
 			{
-				return UE4LogVerbosity.Display;
+				return UELogVerbosity.Display;
 			}
 			else if (string.Compare(From, "Warning", true) == 0)
 			{
-				return UE4LogVerbosity.Warning;
+				return UELogVerbosity.Warning;
 			}
 			else if (string.Compare(From, "Error", true) == 0)
 			{
-				return UE4LogVerbosity.Error;
+				return UELogVerbosity.Error;
 			}
 			else if (string.Compare(From, "Fatal", true) == 0)
 			{
-				return UE4LogVerbosity.Fatal;
+				return UELogVerbosity.Fatal;
 			}
 
 			return Default;
 		}
 
-		public IEnumerable<UE4LogVerbosity> UE4LogVerbosityEnumTypeValues
+		public IEnumerable<UELogVerbosity> UELogVerbosityEnumTypeValues
 		{
 			get
 			{
-				return Enum.GetValues(typeof(UE4LogVerbosity)).Cast<UE4LogVerbosity>();
+				return Enum.GetValues(typeof(UELogVerbosity)).Cast<UELogVerbosity>();
 			}
 		}
 
-		private UE4LogVerbosity _SelectedVerbocityPlugin = UE4LogVerbosity.Log;
-		public UE4LogVerbosity SelectedVerbocityPlugin
+		private UELogVerbosity _SelectedVerbocityPlugin = UELogVerbosity.Log;
+		public UELogVerbosity SelectedVerbocityPlugin
 		{
 			get { return _SelectedVerbocityPlugin; }
 			set
@@ -273,8 +273,8 @@ namespace nDisplayLauncher.Cluster
 			}
 		}
 
-		private UE4LogVerbosity _SelectedVerbocityEngine = UE4LogVerbosity.Log;
-		public UE4LogVerbosity SelectedVerbocityEngine
+		private UELogVerbosity _SelectedVerbocityEngine = UELogVerbosity.Log;
+		public UELogVerbosity SelectedVerbocityEngine
 		{
 			get { return _SelectedVerbocityEngine; }
 			set
@@ -284,8 +284,8 @@ namespace nDisplayLauncher.Cluster
 			}
 		}
 
-		private UE4LogVerbosity _SelectedVerbocityBlueprint = UE4LogVerbosity.Log;
-		public UE4LogVerbosity SelectedVerbocityBlueprint
+		private UELogVerbosity _SelectedVerbocityBlueprint = UELogVerbosity.Log;
+		public UELogVerbosity SelectedVerbocityBlueprint
 		{
 			get { return _SelectedVerbocityBlueprint; }
 			set
@@ -295,8 +295,8 @@ namespace nDisplayLauncher.Cluster
 			}
 		}
 
-		private UE4LogVerbosity _SelectedVerbocityConfig = UE4LogVerbosity.Log;
-		public UE4LogVerbosity SelectedVerbocityConfig
+		private UELogVerbosity _SelectedVerbocityConfig = UELogVerbosity.Log;
+		public UELogVerbosity SelectedVerbocityConfig
 		{
 			get { return _SelectedVerbocityConfig; }
 			set
@@ -306,8 +306,8 @@ namespace nDisplayLauncher.Cluster
 			}
 		}
 
-		private UE4LogVerbosity _SelectedVerbocityCluster = UE4LogVerbosity.Log;
-		public UE4LogVerbosity SelectedVerbocityCluster
+		private UELogVerbosity _SelectedVerbocityCluster = UELogVerbosity.Log;
+		public UELogVerbosity SelectedVerbocityCluster
 		{
 			get { return _SelectedVerbocityCluster; }
 			set
@@ -317,8 +317,8 @@ namespace nDisplayLauncher.Cluster
 			}
 		}
 
-		private UE4LogVerbosity _SelectedVerbocityGame = UE4LogVerbosity.Log;
-		public UE4LogVerbosity SelectedVerbocityGame
+		private UELogVerbosity _SelectedVerbocityGame = UELogVerbosity.Log;
+		public UELogVerbosity SelectedVerbocityGame
 		{
 			get { return _SelectedVerbocityGame; }
 			set
@@ -328,8 +328,8 @@ namespace nDisplayLauncher.Cluster
 			}
 		}
 
-		private UE4LogVerbosity _SelectedVerbocityGameMode = UE4LogVerbosity.Log;
-		public UE4LogVerbosity SelectedVerbocityGameMode
+		private UELogVerbosity _SelectedVerbocityGameMode = UELogVerbosity.Log;
+		public UELogVerbosity SelectedVerbocityGameMode
 		{
 			get { return _SelectedVerbocityGameMode; }
 			set
@@ -339,8 +339,8 @@ namespace nDisplayLauncher.Cluster
 			}
 		}
 
-		private UE4LogVerbosity _SelectedVerbocityInput = UE4LogVerbosity.Log;
-		public UE4LogVerbosity SelectedVerbocityInput
+		private UELogVerbosity _SelectedVerbocityInput = UELogVerbosity.Log;
+		public UELogVerbosity SelectedVerbocityInput
 		{
 			get { return _SelectedVerbocityInput; }
 			set
@@ -350,8 +350,8 @@ namespace nDisplayLauncher.Cluster
 			}
 		}
 
-		private UE4LogVerbosity _SelectedVerbocityVrpn = UE4LogVerbosity.Log;
-		public UE4LogVerbosity SelectedVerbocityVrpn
+		private UELogVerbosity _SelectedVerbocityVrpn = UELogVerbosity.Log;
+		public UELogVerbosity SelectedVerbocityVrpn
 		{
 			get { return _SelectedVerbocityVrpn; }
 			set
@@ -361,8 +361,8 @@ namespace nDisplayLauncher.Cluster
 			}
 		}
 
-		private UE4LogVerbosity _SelectedVerbocityNetwork = UE4LogVerbosity.Log;
-		public UE4LogVerbosity SelectedVerbocityNetwork
+		private UELogVerbosity _SelectedVerbocityNetwork = UELogVerbosity.Log;
+		public UELogVerbosity SelectedVerbocityNetwork
 		{
 			get { return _SelectedVerbocityNetwork; }
 			set
@@ -372,8 +372,8 @@ namespace nDisplayLauncher.Cluster
 			}
 		}
 
-		private UE4LogVerbosity _SelectedVerbocityNetworkMsg = UE4LogVerbosity.Log;
-		public UE4LogVerbosity SelectedVerbocityNetworkMsg
+		private UELogVerbosity _SelectedVerbocityNetworkMsg = UELogVerbosity.Log;
+		public UELogVerbosity SelectedVerbocityNetworkMsg
 		{
 			get { return _SelectedVerbocityNetworkMsg; }
 			set
@@ -383,8 +383,8 @@ namespace nDisplayLauncher.Cluster
 			}
 		}
 
-		private UE4LogVerbosity _SelectedVerbocityRender = UE4LogVerbosity.Log;
-		public UE4LogVerbosity SelectedVerbocityRender
+		private UELogVerbosity _SelectedVerbocityRender = UELogVerbosity.Log;
+		public UELogVerbosity SelectedVerbocityRender
 		{
 			get { return _SelectedVerbocityRender; }
 			set
@@ -394,8 +394,8 @@ namespace nDisplayLauncher.Cluster
 			}
 		}
 
-		private UE4LogVerbosity _SelectedVerbocityRenderSync = UE4LogVerbosity.Log;
-		public UE4LogVerbosity SelectedVerbocityRenderSync
+		private UELogVerbosity _SelectedVerbocityRenderSync = UELogVerbosity.Log;
+		public UELogVerbosity SelectedVerbocityRenderSync
 		{
 			get { return _SelectedVerbocityRenderSync; }
 			set
@@ -491,18 +491,18 @@ namespace nDisplayLauncher.Cluster
 			try
 			{
 				IsCustomLogsUsed            = RegistrySaver.ReadBoolValue(RegistrySaver.RegCategoryLogParams, RegistrySaver.RegLogParamsUseCustomLogs);
-				SelectedVerbocityPlugin     = UE4LogVerbosity_FromString(RegistrySaver.ReadStringValue(RegistrySaver.RegCategoryLogParams, RegistrySaver.RegLogParamsVerbosityPlugin), UE4LogVerbosity.Log);
-				SelectedVerbocityEngine     = UE4LogVerbosity_FromString(RegistrySaver.ReadStringValue(RegistrySaver.RegCategoryLogParams, RegistrySaver.RegLogParamsVerbosityEngine), UE4LogVerbosity.Log);
-				SelectedVerbocityConfig     = UE4LogVerbosity_FromString(RegistrySaver.ReadStringValue(RegistrySaver.RegCategoryLogParams, RegistrySaver.RegLogParamsVerbosityConfig), UE4LogVerbosity.Log);
-				SelectedVerbocityCluster    = UE4LogVerbosity_FromString(RegistrySaver.ReadStringValue(RegistrySaver.RegCategoryLogParams, RegistrySaver.RegLogParamsVerbosityCluster), UE4LogVerbosity.Log);
-				SelectedVerbocityGame       = UE4LogVerbosity_FromString(RegistrySaver.ReadStringValue(RegistrySaver.RegCategoryLogParams, RegistrySaver.RegLogParamsVerbosityGame), UE4LogVerbosity.Log);
-				SelectedVerbocityGameMode   = UE4LogVerbosity_FromString(RegistrySaver.ReadStringValue(RegistrySaver.RegCategoryLogParams, RegistrySaver.RegLogParamsVerbosityGameMode), UE4LogVerbosity.Log);
-				SelectedVerbocityInput      = UE4LogVerbosity_FromString(RegistrySaver.ReadStringValue(RegistrySaver.RegCategoryLogParams, RegistrySaver.RegLogParamsVerbosityInput), UE4LogVerbosity.Log);
-				SelectedVerbocityVrpn       = UE4LogVerbosity_FromString(RegistrySaver.ReadStringValue(RegistrySaver.RegCategoryLogParams, RegistrySaver.RegLogParamsVerbosityInputVrpn), UE4LogVerbosity.Log);
-				SelectedVerbocityNetwork    = UE4LogVerbosity_FromString(RegistrySaver.ReadStringValue(RegistrySaver.RegCategoryLogParams, RegistrySaver.RegLogParamsVerbosityNetwork), UE4LogVerbosity.Log);
-				SelectedVerbocityNetworkMsg = UE4LogVerbosity_FromString(RegistrySaver.ReadStringValue(RegistrySaver.RegCategoryLogParams, RegistrySaver.RegLogParamsVerbosityNetworkMsg), UE4LogVerbosity.Log);
-				SelectedVerbocityBlueprint  = UE4LogVerbosity_FromString(RegistrySaver.ReadStringValue(RegistrySaver.RegCategoryLogParams, RegistrySaver.RegLogParamsVerbosityBlueprint), UE4LogVerbosity.Log);
-				SelectedVerbocityRender     = UE4LogVerbosity_FromString(RegistrySaver.ReadStringValue(RegistrySaver.RegCategoryLogParams, RegistrySaver.RegLogParamsVerbosityRender), UE4LogVerbosity.Log);
+				SelectedVerbocityPlugin     = UELogVerbosity_FromString(RegistrySaver.ReadStringValue(RegistrySaver.RegCategoryLogParams, RegistrySaver.RegLogParamsVerbosityPlugin), UELogVerbosity.Log);
+				SelectedVerbocityEngine     = UELogVerbosity_FromString(RegistrySaver.ReadStringValue(RegistrySaver.RegCategoryLogParams, RegistrySaver.RegLogParamsVerbosityEngine), UELogVerbosity.Log);
+				SelectedVerbocityConfig     = UELogVerbosity_FromString(RegistrySaver.ReadStringValue(RegistrySaver.RegCategoryLogParams, RegistrySaver.RegLogParamsVerbosityConfig), UELogVerbosity.Log);
+				SelectedVerbocityCluster    = UELogVerbosity_FromString(RegistrySaver.ReadStringValue(RegistrySaver.RegCategoryLogParams, RegistrySaver.RegLogParamsVerbosityCluster), UELogVerbosity.Log);
+				SelectedVerbocityGame       = UELogVerbosity_FromString(RegistrySaver.ReadStringValue(RegistrySaver.RegCategoryLogParams, RegistrySaver.RegLogParamsVerbosityGame), UELogVerbosity.Log);
+				SelectedVerbocityGameMode   = UELogVerbosity_FromString(RegistrySaver.ReadStringValue(RegistrySaver.RegCategoryLogParams, RegistrySaver.RegLogParamsVerbosityGameMode), UELogVerbosity.Log);
+				SelectedVerbocityInput      = UELogVerbosity_FromString(RegistrySaver.ReadStringValue(RegistrySaver.RegCategoryLogParams, RegistrySaver.RegLogParamsVerbosityInput), UELogVerbosity.Log);
+				SelectedVerbocityVrpn       = UELogVerbosity_FromString(RegistrySaver.ReadStringValue(RegistrySaver.RegCategoryLogParams, RegistrySaver.RegLogParamsVerbosityInputVrpn), UELogVerbosity.Log);
+				SelectedVerbocityNetwork    = UELogVerbosity_FromString(RegistrySaver.ReadStringValue(RegistrySaver.RegCategoryLogParams, RegistrySaver.RegLogParamsVerbosityNetwork), UELogVerbosity.Log);
+				SelectedVerbocityNetworkMsg = UELogVerbosity_FromString(RegistrySaver.ReadStringValue(RegistrySaver.RegCategoryLogParams, RegistrySaver.RegLogParamsVerbosityNetworkMsg), UELogVerbosity.Log);
+				SelectedVerbocityBlueprint  = UELogVerbosity_FromString(RegistrySaver.ReadStringValue(RegistrySaver.RegCategoryLogParams, RegistrySaver.RegLogParamsVerbosityBlueprint), UELogVerbosity.Log);
+				SelectedVerbocityRender     = UELogVerbosity_FromString(RegistrySaver.ReadStringValue(RegistrySaver.RegCategoryLogParams, RegistrySaver.RegLogParamsVerbosityRender), UELogVerbosity.Log);
 			}
 			catch (Exception ex)
 			{
diff --git a/Engine/Source/Programs/nDisplayLauncher/Cluster/Launcher/MainWindow.Launcher.cs b/Engine/Source/Programs/nDisplayLauncher/Cluster/Launcher/MainWindow.Launcher.cs
--- a/Engine/Source/Programs/nDisplayLauncher/Cluster/Launcher/MainWindow.Launcher.cs
+++ b/Engine/Source/Programs/nDisplayLauncher/Cluster/Launcher/MainWindow.Launcher.cs
@@ -18,8 +18,8 @@ namespace nDisplayLauncher
 	{
 		static readonly string CfgFileExtention = "nDisplay config file (*.cfg)|*.cfg";
 		static readonly string AppFileExtention = "nDisplay application (*.exe)|*.exe|Script (*.bat;*.cmd)|*.bat;*.cmd";
-		static readonly string UE4EditorFileExtention = "UE4 Editor (*.exe)|*.exe";
-		static readonly string UE4ProjectExtention = "UE4 Project (*.uproject)|*.uproject";
+		static readonly string UEEditorFileExtention = "Unreal Editor (*.exe)|*.exe";
+		static readonly string UEProjectExtention = "Unreal Project (*.uproject)|*.uproject";
 
 		private void InitializeLauncher()
 		{
@@ -84,15 +84,15 @@ namespace nDisplayLauncher
 		{
 			string CmdLine = String.Empty;
 			System.Windows.Forms.OpenFileDialog openFileDialog = new System.Windows.Forms.OpenFileDialog();
-			openFileDialog.Filter = UE4EditorFileExtention;
-			openFileDialog.Title = "Select UE4 Editor";
+			openFileDialog.Filter = UEEditorFileExtention;
+			openFileDialog.Title = "Select Unreal Editor";
 			if (openFileDialog.ShowDialog() == System.Windows.Forms.DialogResult.OK)
 			{
 				CmdLine = openFileDialog.FileName;
 
 				System.Windows.Forms.OpenFileDialog openFileDialogForProject = new System.Windows.Forms.OpenFileDialog();
-				openFileDialogForProject.Filter = UE4ProjectExtention;
-				openFileDialogForProject.Title = "Select UE4 Project";
+				openFileDialogForProject.Filter = UEProjectExtention;
+				openFileDialogForProject.Title = "Select Unreal Project";
 				if (openFileDialogForProject.ShowDialog() == System.Windows.Forms.DialogResult.OK)
 				{
 					CmdLine += " " + openFileDialogForProject.FileName + " -game ";
diff --git a/Engine/Source/Programs/nDisplayLauncher/Cluster/Log/MainWindow.Log.cs b/Engine/Source/Programs/nDisplayLauncher/Cluster/Log/MainWindow.Log.cs
--- a/Engine/Source/Programs/nDisplayLauncher/Cluster/Log/MainWindow.Log.cs
+++ b/Engine/Source/Programs/nDisplayLauncher/Cluster/Log/MainWindow.Log.cs
@@ -5,7 +5,7 @@ using System.Windows;
 
 namespace nDisplayLauncher
 {
-	public enum UE4LogVerbosity
+	public enum UELogVerbosity
 	{
 		All = 0,
 		Verbose,
@@ -25,35 +25,35 @@ namespace nDisplayLauncher
 
 		private void ctrlBtnSetVerbosityAll_Click(object sender, RoutedEventArgs e)
 		{
-			UE4LogVerbosity Verbosity = UE4LogVerbosity.Log;
+			UELogVerbosity Verbosity = UELogVerbosity.Log;
 
 			if (sender == ctrlBtnSetVerbosityAll)
 			{
-				Verbosity = UE4LogVerbosity.All;
+				Verbosity = UELogVerbosity.All;
 			}
 			else if (sender == ctrlBtnSetVerbosityVerbose)
 			{
-				Verbosity = UE4LogVerbosity.Verbose;
+				Verbosity = UELogVerbosity.Verbose;
 			}
 			else if (sender == ctrlBtnSetVerbosityLog)
 			{
-				Verbosity = UE4LogVerbosity.Log;
+				Verbosity = UELogVerbosity.Log;
 			}
 			else if (sender == ctrlBtnSetVerbosityDisplay)
 			{
-				Verbosity = UE4LogVerbosity.Display;
+				Verbosity = UELogVerbosity.Display;
 			}
 			else if (sender == ctrlBtnSetVerbosityWarning)
 			{
-				Verbosity = UE4LogVerbosity.Warning;
+				Verbosity = UELogVerbosity.Warning;
 			}
 			else if (sender == ctrlBtnSetVerbosityError)
 			{
-				Verbosity = UE4LogVerbosity.Error;
+				Verbosity = UELogVerbosity.Error;
 			}
 			else if (sender == ctrlBtnSetVerbosityFatal)
 			{
-				Verbosity = UE4LogVerbosity.Fatal;
+				Verbosity = UELogVerbosity.Fatal;
 			}
 			else
 			{
diff --git a/Engine/Source/Programs/nDisplayLauncher/MainWindow.xaml b/Engine/Source/Programs/nDisplayLauncher/MainWindow.xaml
--- a/Engine/Source/Programs/nDisplayLauncher/MainWindow.xaml
+++ b/Engine/Source/Programs/nDisplayLauncher/MainWindow.xaml
@@ -604,43 +604,43 @@
 					</Grid.ColumnDefinitions>
 
 					<Label Grid.Row="0" Grid.Column="0" Content="General plugin" />
-					<ComboBox Grid.Row="0" Grid.Column="1" x:Name="ctrlLogPlugin" SelectedItem="{Binding SelectedVerbocityPlugin}" ItemsSource="{Binding UE4LogVerbosityEnumTypeValues}" Template="{DynamicResource ComboBoxTemplate}" Width="Auto" HorizontalAlignment="Stretch" />
+					<ComboBox Grid.Row="0" Grid.Column="1" x:Name="ctrlLogPlugin" SelectedItem="{Binding SelectedVerbocityPlugin}" ItemsSource="{Binding UELogVerbosityEnumTypeValues}" Template="{DynamicResource ComboBoxTemplate}" Width="Auto" HorizontalAlignment="Stretch" />
 
 					<Label Grid.Row="1" Grid.Column="0" Content="Engine" />
-					<ComboBox Grid.Row="1" Grid.Column="1" x:Name="ctrlLogEngine" SelectedItem="{Binding SelectedVerbocityEngine}" ItemsSource="{Binding UE4LogVerbosityEnumTypeValues}" Template="{DynamicResource ComboBoxTemplate}" Width="Auto" HorizontalAlignment="Stretch" />
+					<ComboBox Grid.Row="1" Grid.Column="1" x:Name="ctrlLogEngine" SelectedItem="{Binding SelectedVerbocityEngine}" ItemsSource="{Binding UELogVerbosityEnumTypeValues}" Template="{DynamicResource ComboBoxTemplate}" Width="Auto" HorizontalAlignment="Stretch" />
 
 					<Label Grid.Row="2" Grid.Column="0" Content="Blueprint" />
-					<ComboBox Grid.Row="2" Grid.Column="1" x:Name="ctrlLogBlueprint" SelectedItem="{Binding SelectedVerbocityBlueprint}" ItemsSource="{Binding UE4LogVerbosityEnumTypeValues}" Template="{DynamicResource ComboBoxTemplate}" Width="Auto" HorizontalAlignment="Stretch" />
+					<ComboBox Grid.Row="2" Grid.Column="1" x:Name="ctrlLogBlueprint" SelectedItem="{Binding SelectedVerbocityBlueprint}" ItemsSource="{Binding UELogVerbosityEnumTypeValues}" Template="{DynamicResource ComboBoxTemplate}" Width="Auto" HorizontalAlignment="Stretch" />
 
 					<Label Grid.Row="3" Grid.Column="0" Content="Config" />
-					<ComboBox Grid.Row="3" Grid.Column="1" x:Name="ctrlLogConfig" SelectedItem="{Binding SelectedVerbocityConfig}" ItemsSource="{Binding UE4LogVerbosityEnumTypeValues}" Template="{DynamicResource ComboBoxTemplate}" Width="Auto" HorizontalAlignment="Stretch" />
+					<ComboBox Grid.Row="3" Grid.Column="1" x:Name="ctrlLogConfig" SelectedItem="{Binding SelectedVerbocityConfig}" ItemsSource="{Binding UELogVerbosityEnumTypeValues}" Template="{DynamicResource ComboBoxTemplate}" Width="Auto" HorizontalAlignment="Stretch" />
 
 					<Label Grid.Row="4" Grid.Column="0" Content="Cluster" />
-					<ComboBox Grid.Row="4" Grid.Column="1" x:Name="ctrlLogCluster" SelectedItem="{Binding SelectedVerbocityCluster}" ItemsSource="{Binding UE4LogVerbosityEnumTypeValues}" Template="{DynamicResource ComboBoxTemplate}" Width="Auto" HorizontalAlignment="Stretch" />
+					<ComboBox Grid.Row="4" Grid.Column="1" x:Name="ctrlLogCluster" SelectedItem="{Binding SelectedVerbocityCluster}" ItemsSource="{Binding UELogVerbosityEnumTypeValues}" Template="{DynamicResource ComboBoxTemplate}" Width="Auto" HorizontalAlignment="Stretch" />
 
 					<Label Grid.Row="5" Grid.Column="0" Content="Game" />
-					<ComboBox Grid.Row="5" Grid.Column="1" x:Name="ctrlLogGame" SelectedItem="{Binding SelectedVerbocityGame}" ItemsSource="{Binding UE4LogVerbosityEnumTypeValues}" Template="{DynamicResource ComboBoxTemplate}" Width="Auto" HorizontalAlignment="Stretch" />
+					<ComboBox Grid.Row="5" Grid.Column="1" x:Name="ctrlLogGame" SelectedItem="{Binding SelectedVerbocityGame}" ItemsSource="{Binding UELogVerbosityEnumTypeValues}" Template="{DynamicResource ComboBoxTemplate}" Width="Auto" HorizontalAlignment="Stretch" />
 
 					<Label Grid.Row="6" Grid.Column="0" Content="Game Mode" />
-					<ComboBox Grid.Row="6" Grid.Column="1" x:Name="ctrlLogGameMode" SelectedItem="{Binding SelectedVerbocityGameMode}" ItemsSource="{Binding UE4LogVerbosityEnumTypeValues}" Template="{DynamicResource ComboBoxTemplate}" Width="Auto" HorizontalAlignment="Stretch" />
+					<ComboBox Grid.Row="6" Grid.Column="1" x:Name="ctrlLogGameMode" SelectedItem="{Binding SelectedVerbocityGameMode}" ItemsSource="{Binding UELogVerbosityEnumTypeValues}" Template="{DynamicResource ComboBoxTemplate}" Width="Auto" HorizontalAlignment="Stretch" />
 
 					<Label Grid.Row="7" Grid.Column="0" Content="Input" />
-					<ComboBox Grid.Row="7" Grid.Column="1" x:Name="ctrlLogInput" SelectedItem="{Binding SelectedVerbocityInput}" ItemsSource="{Binding UE4LogVerbosityEnumTypeValues}" Template="{DynamicResource ComboBoxTemplate}" Width="Auto" HorizontalAlignment="Stretch" />
+					<ComboBox Grid.Row="7" Grid.Column="1" x:Name="ctrlLogInput" SelectedItem="{Binding SelectedVerbocityInput}" ItemsSource="{Binding UELogVerbosityEnumTypeValues}" Template="{DynamicResource ComboBoxTemplate}" Width="Auto" HorizontalAlignment="Stretch" />
 
 					<Label Grid.Row="8" Grid.Column="0" Content="VRPN" />
-					<ComboBox Grid.Row="8" Grid.Column="1" x:Name="ctrlLogVRPN" SelectedItem="{Binding SelectedVerbocityVrpn}" ItemsSource="{Binding UE4LogVerbosityEnumTypeValues}" Template="{DynamicResource ComboBoxTemplate}" Width="Auto" HorizontalAlignment="Stretch" />
+					<ComboBox Grid.Row="8" Grid.Column="1" x:Name="ctrlLogVRPN" SelectedItem="{Binding SelectedVerbocityVrpn}" ItemsSource="{Binding UELogVerbosityEnumTypeValues}" Template="{DynamicResource ComboBoxTemplate}" Width="Auto" HorizontalAlignment="Stretch" />
 
 					<Label Grid.Row="9" Grid.Column="0" Content="Network" />
-					<ComboBox Grid.Row="9" Grid.Column="1" x:Name="ctrlLogNetwork" SelectedItem="{Binding SelectedVerbocityNetwork}" ItemsSource="{Binding UE4LogVerbosityEnumTypeValues}" Template="{DynamicResource ComboBoxTemplate}" Width="Auto" HorizontalAlignment="Stretch" />
+					<ComboBox Grid.Row="9" Grid.Column="1" x:Name="ctrlLogNetwork" SelectedItem="{Binding SelectedVerbocityNetwork}" ItemsSource="{Binding UELogVerbosityEnumTypeValues}" Template="{DynamicResource ComboBoxTemplate}" Width="Auto" HorizontalAlignment="Stretch" />
 
 					<Label Grid.Row="10" Grid.Column="0" Content="Network Messages" />
-					<ComboBox Grid.Row="10" Grid.Column="1" x:Name="ctrlLogNetworkMsg" SelectedItem="{Binding SelectedVerbocityNetworkMsg}" ItemsSource="{Binding UE4LogVerbosityEnumTypeValues}" Template="{DynamicResource ComboBoxTemplate}" Width="Auto" HorizontalAlignment="Stretch" />
+					<ComboBox Grid.Row="10" Grid.Column="1" x:Name="ctrlLogNetworkMsg" SelectedItem="{Binding SelectedVerbocityNetworkMsg}" ItemsSource="{Binding UELogVerbosityEnumTypeValues}" Template="{DynamicResource ComboBoxTemplate}" Width="Auto" HorizontalAlignment="Stretch" />
 
 					<Label Grid.Row="11" Grid.Column="0" Content="Render" />
-					<ComboBox Grid.Row="11" Grid.Column="1" x:Name="ctrlLogRender" SelectedItem="{Binding SelectedVerbocityRender}" ItemsSource="{Binding UE4LogVerbosityEnumTypeValues}" Template="{DynamicResource ComboBoxTemplate}" Width="Auto" HorizontalAlignment="Stretch" />
+					<ComboBox Grid.Row="11" Grid.Column="1" x:Name="ctrlLogRender" SelectedItem="{Binding SelectedVerbocityRender}" ItemsSource="{Binding UELogVerbosityEnumTypeValues}" Template="{DynamicResource ComboBoxTemplate}" Width="Auto" HorizontalAlignment="Stretch" />
 
 					<Label Grid.Row="12" Grid.Column="0" Content="RenderSync" />
-					<ComboBox Grid.Row="12" Grid.Column="1" x:Name="ctrlLogRenderSync" SelectedItem="{Binding SelectedVerbocityRenderSync}" ItemsSource="{Binding UE4LogVerbosityEnumTypeValues}" Template="{DynamicResource ComboBoxTemplate}" Width="Auto" HorizontalAlignment="Stretch" />
+					<ComboBox Grid.Row="12" Grid.Column="1" x:Name="ctrlLogRenderSync" SelectedItem="{Binding SelectedVerbocityRenderSync}" ItemsSource="{Binding UELogVerbosityEnumTypeValues}" Template="{DynamicResource ComboBoxTemplate}" Width="Auto" HorizontalAlignment="Stretch" />
 
 					<Label Grid.Row="0" Grid.Column="3" Content="For each:" />
 					<Button x:Name="ctrlBtnSetVerbosityAll" Grid.Row="1" Grid.Column="3" Content="All" Click="ctrlBtnSetVerbosityAll_Click" />
